mac CIS fix
#ue4

diff --git a/Engine/Plugins/Media/MediaAssetEditor/Source/MediaAssetEditor/Private/MediaAssetEditorPrivatePCH.h b/Engine/Plugins/Media/MediaAssetEditor/Source/MediaAssetEditor/Private/MediaAssetEditorPrivatePCH.h
--- a/Engine/Plugins/Media/MediaAssetEditor/Source/MediaAssetEditor/Private/MediaAssetEditorPrivatePCH.h
+++ b/Engine/Plugins/Media/MediaAssetEditor/Source/MediaAssetEditor/Private/MediaAssetEditorPrivatePCH.h
@@ -6,7 +6,12 @@
 /* Private dependencies
  *****************************************************************************/
 
+#include "CoreUObject.h"
 #include "AssetToolsModule.h"
+#include "Slate.h"
+#include "Engine.h"
+#include "SlateStyle.h"
+#include "EditorStyle.h"
 #include "IDetailsView.h"
 #include "Media.h"
 #include "MediaAsset.h"
@@ -15,10 +20,7 @@
 #include "ModuleInterface.h"
 #include "PropertyEditorModule.h"
 #include "SceneViewport.h"
-#include "Slate.h"
 #include "SlateTextures.h"
-#include "SlateStyle.h"
-#include "EditorStyle.h"
 #include "TickableObjectRenderThread.h"
 #include "WorkspaceMenuStructureModule.h"
 #include "Toolkits/AssetEditorToolkit.h"
