Media: Fixed shutdown crash when unregistering thumbnail renderer

diff --git a/Engine/Plugins/Media/MediaAssetEditor/Source/MediaAssetEditor/Private/MediaAssetEditorModule.cpp b/Engine/Plugins/Media/MediaAssetEditor/Source/MediaAssetEditor/Private/MediaAssetEditorModule.cpp
--- a/Engine/Plugins/Media/MediaAssetEditor/Source/MediaAssetEditor/Private/MediaAssetEditorModule.cpp
+++ b/Engine/Plugins/Media/MediaAssetEditor/Source/MediaAssetEditor/Private/MediaAssetEditorModule.cpp
@@ -121,7 +121,10 @@ protected:
 	/** Unregisters all asset thumbnail renderers. */
 	void UnregisterThumbnailRenderers( )
 	{
-		UThumbnailManager::Get().UnregisterCustomRenderer(UMediaTexture::StaticClass());
+		if (UObjectInitialized())
+		{
+			UThumbnailManager::Get().UnregisterCustomRenderer(UMediaTexture::StaticClass());
+		}
 	}
 
 private:
