Media: removed the 'Stop' feature as it is redundant. Play and Pause are sufficient.

diff --git a/Engine/Plugins/Media/MediaAssetEditor/Source/MediaAssetEditor/Private/AssetTools/MediaAssetActions.cpp b/Engine/Plugins/Media/MediaAssetEditor/Source/MediaAssetEditor/Private/AssetTools/MediaAssetActions.cpp
--- a/Engine/Plugins/Media/MediaAssetEditor/Source/MediaAssetEditor/Private/AssetTools/MediaAssetActions.cpp
+++ b/Engine/Plugins/Media/MediaAssetEditor/Source/MediaAssetEditor/Private/AssetTools/MediaAssetActions.cpp
@@ -50,16 +50,6 @@ void FMediaAssetActions::GetActions( const TArray<UObject*>& InObjects, FMenuBui
 		)
 	);
 
-	MenuBuilder.AddMenuEntry(
-		LOCTEXT("MediaAsset_StopMovie", "Stop Movie"),
-		LOCTEXT("MediaAsset_StopMovieToolTip", "Stops playback of the media."),
-		FSlateIcon(),
-		FUIAction(
-			FExecuteAction::CreateSP(this, &FMediaAssetActions::HandleStopMovieActionExecute, MediaAssets),
-			FCanExecuteAction()
-		)
-	);
-
 	MenuBuilder.AddMenuSeparator();
 
 	MenuBuilder.AddMenuEntry(
@@ -280,16 +270,4 @@ void FMediaAssetActions::HandlePlayMovieActionExecute( TArray<TWeakObjectPtr<UMe
 }
 
 
-void FMediaAssetActions::HandleStopMovieActionExecute( TArray<TWeakObjectPtr<UMediaAsset>> Objects )
-{
-	for (auto MediaAsset : Objects)
-	{
-		if (MediaAsset.IsValid())
-		{
-			MediaAsset->Stop();
-		}
-	}
-}
-
-
 #undef LOCTEXT_NAMESPACE
diff --git a/Engine/Plugins/Media/MediaAssetEditor/Source/MediaAssetEditor/Private/AssetTools/MediaAssetActions.h b/Engine/Plugins/Media/MediaAssetEditor/Source/MediaAssetEditor/Private/AssetTools/MediaAssetActions.h
--- a/Engine/Plugins/Media/MediaAssetEditor/Source/MediaAssetEditor/Private/AssetTools/MediaAssetActions.h
+++ b/Engine/Plugins/Media/MediaAssetEditor/Source/MediaAssetEditor/Private/AssetTools/MediaAssetActions.h
@@ -48,9 +48,6 @@ private:
 	// Callback for executing the 'Play Movie' action.
 	void HandlePlayMovieActionExecute( TArray<TWeakObjectPtr<UMediaAsset>> Objects );
 
-	// Callback for executing the 'Play Movie' action.
-	void HandleStopMovieActionExecute( TArray<TWeakObjectPtr<UMediaAsset>> Objects );
-
 private:
 
 	// Pointer to the style set to use for toolkits.
diff --git a/Engine/Plugins/Media/MediaAssetEditor/Source/MediaAssetEditor/Private/MediaAssetEditorToolkit.cpp b/Engine/Plugins/Media/MediaAssetEditor/Source/MediaAssetEditor/Private/MediaAssetEditorToolkit.cpp
--- a/Engine/Plugins/Media/MediaAssetEditor/Source/MediaAssetEditor/Private/MediaAssetEditorToolkit.cpp
+++ b/Engine/Plugins/Media/MediaAssetEditor/Source/MediaAssetEditor/Private/MediaAssetEditorToolkit.cpp
@@ -210,11 +210,6 @@ void FMediaAssetEditorToolkit::BindCommands( )
 		Commands.RewindMedia,
 		FExecuteAction::CreateSP(this, &FMediaAssetEditorToolkit::HandleRewindMediaActionExecute),
 		FCanExecuteAction::CreateSP(this, &FMediaAssetEditorToolkit::HandleRewindMediaActionCanExecute));
-
-	ToolkitCommands->MapAction(
-		Commands.StopMedia,
-		FExecuteAction::CreateSP(this, &FMediaAssetEditorToolkit::HandleStopMediaActionExecute),
-		FCanExecuteAction::CreateSP(this, &FMediaAssetEditorToolkit::HandleStopMediaActionCanExecute));
 }
 
 
@@ -231,7 +226,6 @@ void FMediaAssetEditorToolkit::ExtendToolBar( )
 				ToolbarBuilder.AddToolBarButton(FMediaAssetEditorCommands::Get().ReverseMedia);
 				ToolbarBuilder.AddToolBarButton(FMediaAssetEditorCommands::Get().PlayMedia);
 				ToolbarBuilder.AddToolBarButton(FMediaAssetEditorCommands::Get().PauseMedia);
-				ToolbarBuilder.AddToolBarButton(FMediaAssetEditorCommands::Get().StopMedia);
 				ToolbarBuilder.AddToolBarButton(FMediaAssetEditorCommands::Get().ForwardMedia);
 			}
 			ToolbarBuilder.EndSection();
@@ -342,19 +336,6 @@ void FMediaAssetEditorToolkit::HandleRewindMediaActionExecute( )
 }
 
 
-bool FMediaAssetEditorToolkit::HandleStopMediaActionCanExecute( ) const
-{
-	return false;
-//	return MediaAsset->CanStop();
-}
-
-
-void FMediaAssetEditorToolkit::HandleStopMediaActionExecute( )
-{
-	MediaAsset->Stop();
-}
-
-
 TSharedRef<SDockTab> FMediaAssetEditorToolkit::HandleTabManagerSpawnTab( const FSpawnTabArgs& Args, FName TabIdentifier )
 {
 	TSharedPtr<SWidget> TabWidget = SNullWidget::NullWidget;
diff --git a/Engine/Plugins/Media/MediaAssetEditor/Source/MediaAssetEditor/Private/MediaAssetEditorToolkit.h b/Engine/Plugins/Media/MediaAssetEditor/Source/MediaAssetEditor/Private/MediaAssetEditorToolkit.h
--- a/Engine/Plugins/Media/MediaAssetEditor/Source/MediaAssetEditor/Private/MediaAssetEditorToolkit.h
+++ b/Engine/Plugins/Media/MediaAssetEditor/Source/MediaAssetEditor/Private/MediaAssetEditorToolkit.h
@@ -126,12 +126,6 @@ private:
 	// Callback for executing the ReverseMedia action.
 	void HandleRewindMediaActionExecute( );
 
-	// Callback for determining whether the StopMedia action can execute.
-	bool HandleStopMediaActionCanExecute( ) const;
-
-	// Callback for executing the StopMedia action.
-	void HandleStopMediaActionExecute( );
-
 	// Callback for spawning the Properties tab.
 	TSharedRef<SDockTab> HandleTabManagerSpawnTab( const FSpawnTabArgs& Args, FName TabIdentifier );
 
diff --git a/Engine/Plugins/Media/MediaAssetEditor/Source/MediaAssetEditor/Private/Models/MediaAssetEditorCommands.cpp b/Engine/Plugins/Media/MediaAssetEditor/Source/MediaAssetEditor/Private/Models/MediaAssetEditorCommands.cpp
--- a/Engine/Plugins/Media/MediaAssetEditor/Source/MediaAssetEditor/Private/Models/MediaAssetEditorCommands.cpp
+++ b/Engine/Plugins/Media/MediaAssetEditor/Source/MediaAssetEditor/Private/Models/MediaAssetEditorCommands.cpp
@@ -10,5 +10,4 @@ void FMediaAssetEditorCommands::RegisterCommands()
 	UI_COMMAND(PlayMedia, "Play", "Start media playback", EUserInterfaceActionType::Button, FInputGesture());
 	UI_COMMAND(ReverseMedia, "Reverse", "Reverses media playback", EUserInterfaceActionType::Button, FInputGesture());
 	UI_COMMAND(RewindMedia, "Rewind", "Rewinds the media to the beginning", EUserInterfaceActionType::Button, FInputGesture());
-	UI_COMMAND(StopMedia, "Stop", "Stop media playback", EUserInterfaceActionType::Button, FInputGesture());
 }
diff --git a/Engine/Plugins/Media/MediaAssetEditor/Source/MediaAssetEditor/Private/Models/MediaAssetEditorCommands.h b/Engine/Plugins/Media/MediaAssetEditor/Source/MediaAssetEditor/Private/Models/MediaAssetEditorCommands.h
--- a/Engine/Plugins/Media/MediaAssetEditor/Source/MediaAssetEditor/Private/Models/MediaAssetEditorCommands.h
+++ b/Engine/Plugins/Media/MediaAssetEditor/Source/MediaAssetEditor/Private/Models/MediaAssetEditorCommands.h
@@ -40,7 +40,4 @@ public:
 
 	/** Rewinds the media to the beginning. */
 	TSharedPtr<FUICommandInfo> RewindMedia;
-
-	/** Stops media playback. */
-	TSharedPtr<FUICommandInfo> StopMedia;
 };
diff --git a/Engine/Source/Runtime/MediaAssets/Private/Assets/MediaAsset.cpp b/Engine/Source/Runtime/MediaAssets/Private/Assets/MediaAsset.cpp
--- a/Engine/Source/Runtime/MediaAssets/Private/Assets/MediaAsset.cpp
+++ b/Engine/Source/Runtime/MediaAssets/Private/Assets/MediaAsset.cpp
@@ -31,12 +31,6 @@ bool UMediaAsset::CanPlay( ) const
 }
 
 
-bool UMediaAsset::CanStop( ) const
-{
-	return MediaPlayer.IsValid() && MediaPlayer->IsReady();
-}
-
-
 FTimespan UMediaAsset::GetDuration( ) const
 {
 	if (MediaPlayer.IsValid())
@@ -152,16 +146,6 @@ bool UMediaAsset::Seek( const FTimespan& InTime )
 }
 
 
-void UMediaAsset::Stop( )
-{
-	if (MediaPlayer.IsValid())
-	{
-		Rewind();
-		Pause();
-	}
-}
-
-
 bool UMediaAsset::SupportsRate( float Rate, bool Unthinned ) const
 {
 	return MediaPlayer.IsValid() && MediaPlayer->SupportsRate(Rate, Unthinned);
diff --git a/Engine/Source/Runtime/MediaAssets/Private/MediaAssetsModule.cpp b/Engine/Source/Runtime/MediaAssets/Private/MediaAssetsModule.cpp
--- a/Engine/Source/Runtime/MediaAssets/Private/MediaAssetsModule.cpp
+++ b/Engine/Source/Runtime/MediaAssets/Private/MediaAssetsModule.cpp
@@ -17,7 +17,7 @@ public:
 
 	virtual bool Exec( UWorld* InWorld, const TCHAR* Cmd, FOutputDevice& Ar ) override
 	{
-		if (FParse::Command(&Cmd, TEXT("MOVIE")))
+		if (FParse::Command(&Cmd, TEXT("MEDIA")))
 		{
 			FString MovieCmd = FParse::Token(Cmd, 0);
 
@@ -37,14 +37,6 @@ public:
 					MediaAsset->Pause();
 				}
 			}
-			else if (MovieCmd.Contains(TEXT("STOP")))
-			{
-				for (TObjectIterator<UMediaAsset> It; It; ++It)
-				{
-					UMediaAsset* MediaAsset = *It;
-					MediaAsset->Stop();
-				}
-			}
 
 			return true;
 		}
diff --git a/Engine/Source/Runtime/MediaAssets/Public/MediaAsset.h b/Engine/Source/Runtime/MediaAssets/Public/MediaAsset.h
--- a/Engine/Source/Runtime/MediaAssets/Public/MediaAsset.h
+++ b/Engine/Source/Runtime/MediaAssets/Public/MediaAsset.h
@@ -48,7 +48,7 @@ public:
 	 * Playback can be paused if the media supports pausing and if it is currently playing.
 	 *
 	 * @return true if pausing playback can be paused, false otherwise.
-	 * @see CanPlay, CanStop, Pause
+	 * @see CanPlay, Pause
 	 */
 	UFUNCTION(BlueprintCallable, Category="Media|MediaAsset")
 	bool CanPause( ) const;
@@ -57,20 +57,11 @@ public:
 	 * Checks whether media playback can be started right now.
 	 *
 	 * @return true if playback can be started, false otherwise.
-	 * @see CanPause, CanStop, Play
+	 * @see CanPause, Play
 	 */
 	UFUNCTION(BlueprintCallable, Category="Media|MediaAsset")
 	bool CanPlay( ) const;
 
-	/**
-	 * Checks whether playback can be stopped.
-	 *
-	 * @return true if stopping is supported, false otherwise.
-	 * @see CanPause, CanPlay, Stop
-	 */
-	UFUNCTION(BlueprintCallable, Category="Media|MediaAsset")
-	bool CanStop( ) const;
-
 	/**
 	 * Gets the media's duration.
 	 *
@@ -137,7 +128,7 @@ public:
 	 * Checks whether playback has stopped.
 	 *
 	 * @return true if playback has stopped, false otherwise.
-	 * @see CanStop, IsPaused, IsPlaying, Stop
+	 * @see IsPaused, IsPlaying, Stop
 	 */
 	UFUNCTION(BlueprintCallable, Category="Media|MediaAsset")
 	bool IsStopped( ) const;
@@ -209,14 +200,6 @@ public:
 	UFUNCTION(BlueprintCallable, Category="Media|MediaAsset")
 	bool SetRate( float InRate );
 
-	/**
-	 * Stops movie playback.
-	 *
-	 * @see CanStop, IsStopped, Pause, Play, Seek
-	 */
-	UFUNCTION(BlueprintCallable, Category="Media|MediaAsset")
-	void Stop( );
-
 	/**
 	 * Checks whether the specified playback rate is supported.
 	 *
