Media: Added tooltips for elapsed / remaining time indicators in Media Asset Editor

diff --git a/Engine/Plugins/Media/MediaAssetEditor/Source/MediaAssetEditor/Private/Widgets/SMediaAssetEditorPlayer.cpp b/Engine/Plugins/Media/MediaAssetEditor/Source/MediaAssetEditor/Private/Widgets/SMediaAssetEditorPlayer.cpp
--- a/Engine/Plugins/Media/MediaAssetEditor/Source/MediaAssetEditor/Private/Widgets/SMediaAssetEditorPlayer.cpp
+++ b/Engine/Plugins/Media/MediaAssetEditor/Source/MediaAssetEditor/Private/Widgets/SMediaAssetEditorPlayer.cpp
@@ -182,6 +182,7 @@ void SMediaAssetEditorPlayer::Construct( const FArguments& InArgs, UMediaAsset*
 					[
 						SNew(STextBlock)
 							.Text(this, &SMediaAssetEditorPlayer::HandleElapsedTimeTextBlockText)
+							.ToolTipText(LOCTEXT("ElapsedTimeTooltip", "Elapsed Time"))
 					]
 
 				// scrubber
@@ -203,7 +204,8 @@ void SMediaAssetEditorPlayer::Construct( const FArguments& InArgs, UMediaAsset*
 					.AutoWidth()
 					[
 						SNew(STextBlock)
-							.Text(this, &SMediaAssetEditorPlayer::HandleRemainingTimeTextBlockText)					
+							.Text(this, &SMediaAssetEditorPlayer::HandleRemainingTimeTextBlockText)
+							.ToolTipText(LOCTEXT("RemainingTimeTooltip", "Remaining Time"))
 					]
 			]
 	];
