Media: Added entry box for current playback rate to MediaAsset Editor

diff --git a/Engine/Plugins/Media/MediaAssetEditor/Source/MediaAssetEditor/Private/Widgets/SMediaAssetEditorPlayer.cpp b/Engine/Plugins/Media/MediaAssetEditor/Source/MediaAssetEditor/Private/Widgets/SMediaAssetEditorPlayer.cpp
--- a/Engine/Plugins/Media/MediaAssetEditor/Source/MediaAssetEditor/Private/Widgets/SMediaAssetEditorPlayer.cpp
+++ b/Engine/Plugins/Media/MediaAssetEditor/Source/MediaAssetEditor/Private/Widgets/SMediaAssetEditorPlayer.cpp
@@ -34,8 +34,7 @@ void SMediaAssetEditorPlayer::Construct( const FArguments& InArgs, UMediaAsset*
 				SNew(SHorizontalBox)
 
 				+ SHorizontalBox::Slot()
-					.FillWidth(0.33f)
-					.HAlign(HAlign_Left)
+					.AutoWidth()
 					[
 						// video track selector
 						SNew(SHorizontalBox)
@@ -66,8 +65,8 @@ void SMediaAssetEditorPlayer::Construct( const FArguments& InArgs, UMediaAsset*
 					]
 
 				+ SHorizontalBox::Slot()
-					.FillWidth(0.33f)
-					.HAlign(HAlign_Left)
+					.AutoWidth()
+					.Padding(16.0f, 0.0f, 0.0f, 0.0f)
 					[
 						// audio track selector
 						SNew(SHorizontalBox)
@@ -82,7 +81,7 @@ void SMediaAssetEditorPlayer::Construct( const FArguments& InArgs, UMediaAsset*
 
 						+ SHorizontalBox::Slot()
 							.FillWidth(1.0f)
-							.Padding(4.0f, 0.0f)
+							.Padding(4.0f, 0.0f, 0.0f, 0.0f)
 							.VAlign(VAlign_Center)
 							[
 								SAssignNew(AudioTrackComboBox, SComboBox<IMediaTrackPtr>)
@@ -98,8 +97,8 @@ void SMediaAssetEditorPlayer::Construct( const FArguments& InArgs, UMediaAsset*
 					]
 
 				+ SHorizontalBox::Slot()
-					.FillWidth(0.33f)
-					.HAlign(HAlign_Left)
+					.AutoWidth()
+					.Padding(16.0f, 0.0f, 0.0f, 0.0f)
 					[
 						// caption track selector
 						SNew(SHorizontalBox)
@@ -114,7 +113,7 @@ void SMediaAssetEditorPlayer::Construct( const FArguments& InArgs, UMediaAsset*
 
 						+ SHorizontalBox::Slot()
 							.FillWidth(1.0f)
-							.Padding(4.0f, 0.0f)
+							.Padding(4.0f, 0.0f, 0.0f, 0.0f)
 							.VAlign(VAlign_Center)
 							[
 								SAssignNew(CaptionTrackComboBox, SComboBox<IMediaTrackPtr>)
@@ -128,6 +127,36 @@ void SMediaAssetEditorPlayer::Construct( const FArguments& InArgs, UMediaAsset*
 									]
 							]
 					]
+
+				+ SHorizontalBox::Slot()
+					.FillWidth(1.0f)
+					.HAlign(HAlign_Right)
+					.Padding(16.0f, 0.0f, 0.0f, 0.0f)
+					[
+						// playback rate spin box
+						SNew(SHorizontalBox)
+
+						+ SHorizontalBox::Slot()
+							.AutoWidth()
+							.VAlign(VAlign_Center)
+							[
+								SNew(STextBlock)
+									.Text(LOCTEXT("CaptionTrackLabel", "Playback Rate:"))
+							]
+
+						+ SHorizontalBox::Slot()
+							.FillWidth(1.0f)
+							.Padding(4.0f, 0.0f, 0.0f, 0.0f)
+							.VAlign(VAlign_Center)
+							[
+								SNew(SNumericEntryBox<float>)
+									.Delta(1.0f)
+									.MaxValue(this, &SMediaAssetEditorPlayer::HandlePlaybackRateBoxMaxValue)
+									.MinValue(this, &SMediaAssetEditorPlayer::HandlePlaybackRateBoxMinValue)
+									.Value(this, &SMediaAssetEditorPlayer::HandlePlaybackRateSpinBoxValue)
+									.OnValueChanged(this, &SMediaAssetEditorPlayer::HandlePlaybackRateBoxValueChanged)
+							]
+					]
 			]
 
 		+ SVerticalBox::Slot()
@@ -426,6 +455,44 @@ FText SMediaAssetEditorPlayer::HandleOverlayStateText( ) const
 }
 
 
+TOptional<float> SMediaAssetEditorPlayer::HandlePlaybackRateBoxMaxValue( ) const
+{
+	IMediaPlayerPtr MediaPlayer = MediaAsset->GetMediaPlayer();
+
+	if (MediaPlayer.IsValid())
+	{
+		return MediaPlayer->GetMediaInfo().GetSupportedRates(EMediaPlaybackDirections::Forward, false).GetUpperBoundValue();
+	}
+
+	return 0.0f;
+}
+
+
+TOptional<float> SMediaAssetEditorPlayer::HandlePlaybackRateBoxMinValue( ) const
+{
+	IMediaPlayerPtr MediaPlayer = MediaAsset->GetMediaPlayer();
+
+	if (MediaPlayer.IsValid())
+	{
+		return MediaPlayer->GetMediaInfo().GetSupportedRates(EMediaPlaybackDirections::Reverse, false).GetLowerBoundValue();
+	}
+
+	return 0.0f;
+}
+
+
+TOptional<float> SMediaAssetEditorPlayer::HandlePlaybackRateSpinBoxValue( ) const
+{
+	return MediaAsset->GetRate();
+}
+
+
+void SMediaAssetEditorPlayer::HandlePlaybackRateBoxValueChanged( float NewValue )
+{
+	MediaAsset->SetRate(NewValue);
+}
+
+
 bool SMediaAssetEditorPlayer::HandlePositionSliderIsEnabled( ) const
 {
 	return MediaAsset->SupportsSeeking();
diff --git a/Engine/Plugins/Media/MediaAssetEditor/Source/MediaAssetEditor/Private/Widgets/SMediaAssetEditorPlayer.h b/Engine/Plugins/Media/MediaAssetEditor/Source/MediaAssetEditor/Private/Widgets/SMediaAssetEditorPlayer.h
--- a/Engine/Plugins/Media/MediaAssetEditor/Source/MediaAssetEditor/Private/Widgets/SMediaAssetEditorPlayer.h
+++ b/Engine/Plugins/Media/MediaAssetEditor/Source/MediaAssetEditor/Private/Widgets/SMediaAssetEditorPlayer.h
@@ -66,6 +66,18 @@ private:
 	// Callback for getting the text of the playback state overlay.
 	FText HandleOverlayStateText( ) const;
 
+	// Callback for getting the maximum value of the playback rate spin box.
+	TOptional<float> HandlePlaybackRateBoxMaxValue( ) const;
+
+	// Callback for getting the minimum value of the playback rate spin box.
+	TOptional<float> HandlePlaybackRateBoxMinValue( ) const;
+
+	// Callback for getting the current value of the playback rate spin box.
+	TOptional<float> HandlePlaybackRateSpinBoxValue( ) const;
+
+	// Callback for committing a new value to the playback rate spin box.
+	void HandlePlaybackRateBoxValueChanged( float NewValue );
+
 	// Callback for getting the enabled state of the position slider.
 	bool HandlePositionSliderIsEnabled( ) const;
 
